A bunch of gcc non-unity CIS fixes, mostly missing headers

diff --git a/Engine/Plugins/Online/OnlineSubsystemTencent/Source/Private/OnlineSubsystemTencent.cpp b/Engine/Plugins/Online/OnlineSubsystemTencent/Source/Private/OnlineSubsystemTencent.cpp
--- a/Engine/Plugins/Online/OnlineSubsystemTencent/Source/Private/OnlineSubsystemTencent.cpp
+++ b/Engine/Plugins/Online/OnlineSubsystemTencent/Source/Private/OnlineSubsystemTencent.cpp
@@ -26,6 +26,11 @@
 #include "OnlineAsyncTasksTencent.h"
 #include "PlayTimeLimitImpl.h"
 #include "OnlinePlayTimeLimitTencent.h"
+#include "NotForLicensees/OnlineSessionTencentTCLS.h"
+#include "OnlineIdentityTencent.h"
+#include "OnlineSubsystemTencent.h"
+#include "NotForLicensees/OnlineDirectoryTencent.h"
+
 
 #if WITH_TENCENTSDK
 
@@ -291,10 +296,12 @@ bool FOnlineSubsystemTencent::Init()
 	TencentStore = MakeShared<FOnlineStoreTencent, ESPMode::ThreadSafe>(this);
 #endif
 
+#if WITH_TENCENTSDK
 	// Get the server info from TDIR at startup
 	ExecuteNextTick([this]() {
 		TencentDirectory->QueryServerInfo();
 	});
+#endif // WITH_TENCENTSDK
 
 	// update services based on user login/logout events
 	OnLoginChangedHandle = TencentIdentity->AddOnLoginChangedDelegate_Handle(FOnLoginChangedDelegate::CreateThreadSafeSP(this, &FOnlineSubsystemTencent::OnLoginChanged));
diff --git a/Engine/Plugins/Online/OnlineSubsystemTencent/Source/Public/OnlineSubsystemTencent.h b/Engine/Plugins/Online/OnlineSubsystemTencent/Source/Public/OnlineSubsystemTencent.h
--- a/Engine/Plugins/Online/OnlineSubsystemTencent/Source/Public/OnlineSubsystemTencent.h
+++ b/Engine/Plugins/Online/OnlineSubsystemTencent/Source/Public/OnlineSubsystemTencent.h
@@ -172,8 +172,10 @@ private:
 	FOnlineIdentityTencentPtr TencentIdentity;
 	/** Interface to the session services (needed for TSS anticheat handling) */
 	FOnlineSessionTencentPtr TencentSession;
+#if WITH_TENCENTSDK
 	/** Interface to the directory service (TCLS proxy TDIR info) */
 	FOnlineDirectoryTencentPtr TencentDirectory;
+#endif // WITH_TENCENTSDK
 #if WITH_TENCENT_RAIL_SDK
 	/** Interface to the friends services */
 	FOnlineFriendsTencentPtr TencentFriends;
