- Removing NFL code from OSSTencent

diff --git a/Engine/Plugins/Online/OnlineSubsystemTencent/Source/Private/OnlineIdentityTencentCommon.cpp b/Engine/Plugins/Online/OnlineSubsystemTencent/Source/Private/OnlineIdentityTencentCommon.cpp
--- a/Engine/Plugins/Online/OnlineSubsystemTencent/Source/Private/OnlineIdentityTencentCommon.cpp
+++ b/Engine/Plugins/Online/OnlineSubsystemTencent/Source/Private/OnlineIdentityTencentCommon.cpp
@@ -3,7 +3,7 @@
 #include "OnlineIdentityTencent.h"
 #include "OnlineSubsystemTencentPCH.h"
 
-#if WITH_TENCENT_RAIL_SDK || WITH_TENCENT_TCLS
+#if WITH_TENCENT_RAIL_SDK
 bool FOnlineIdentityTencent::GetLocalUserIdx(const FUniqueNetId& UserId, int32& OutLocalIdx) const
 {
 	for (TMap<int32, TSharedPtr<const FUniqueNetId> >::TConstIterator It(UserIds); It; ++It)
@@ -19,7 +19,7 @@ bool FOnlineIdentityTencent::GetLocalUserIdx(const FUniqueNetId& UserId, int32&
 #endif
 
 
-#if !WITH_TENCENT_RAIL_SDK && !WITH_TENCENT_TCLS
+#if !WITH_TENCENT_RAIL_SDK
 FOnlineIdentityTencent::FOnlineIdentityTencent(FOnlineSubsystemTencent* InSubsystem) {}
 FOnlineIdentityTencent::~FOnlineIdentityTencent() {}
 TSharedPtr<FUserOnlineAccountTencent> FOnlineIdentityTencent::GetUserAccountTencent(const FUniqueNetId& UserId) const { return nullptr; }
@@ -42,4 +42,4 @@ FString FOnlineIdentityTencent::GetAuthType() const{ return FString(); }
 void FOnlineIdentityTencent::RevokeAuthToken(const FUniqueNetId& UserId, const FOnRevokeAuthTokenCompleteDelegate& Delegate){ }
 bool FOnlineIdentityTencent::GetLocalUserIdx(const FUniqueNetId& UserId, int32& OutLocalIdx) const { return false; }
 
-#endif // !WITH_TENCENT_RAIL_SDK && !WITH_TENCENT_TCLS
\ No newline at end of file
+#endif // !WITH_TENCENT_RAIL_SDK
\ No newline at end of file
diff --git a/Engine/Plugins/Online/OnlineSubsystemTencent/Source/Private/OnlineSubsystemTencent.cpp b/Engine/Plugins/Online/OnlineSubsystemTencent/Source/Private/OnlineSubsystemTencent.cpp
--- a/Engine/Plugins/Online/OnlineSubsystemTencent/Source/Private/OnlineSubsystemTencent.cpp
+++ b/Engine/Plugins/Online/OnlineSubsystemTencent/Source/Private/OnlineSubsystemTencent.cpp
@@ -5,12 +5,6 @@
 #include "RailSdkWrapper.h"
 #include "OnlineAsyncTaskManagerTencent.h"
 #include "OnlineIdentityTencent.h"
-#if WITH_TENCENT_TCLS
-#include "NotForLicensees/OnlineSessionTencentTCLS.h"
-#include "NotForLicensees/OnlineDirectoryTencent.h"
-#include "NotForLicensees/TclsProxyConnection.h"
-#include "NotForLicensees/CrossVoiceClient.h"
-#endif 
 #include "OnlineSessionTencentRail.h"
 #include "OnlineFriendsTencent.h"
 #include "OnlinePresenceTencent.h"
@@ -26,10 +20,8 @@
 #include "OnlineAsyncTasksTencent.h"
 #include "PlayTimeLimitImpl.h"
 #include "OnlinePlayTimeLimitTencent.h"
-#include "NotForLicensees/OnlineSessionTencentTCLS.h"
 #include "OnlineIdentityTencent.h"
 #include "OnlineSubsystemTencent.h"
-#include "NotForLicensees/OnlineDirectoryTencent.h"
 
 
 #if WITH_TENCENTSDK
@@ -217,14 +209,6 @@ bool FOnlineSubsystemTencent::Init()
 {
 	UE_LOG_ONLINE(Verbose, TEXT("FOnlineSubsystemTencent::Init() Name: %s"), *InstanceName.ToString());
 
-#if WITH_TENCENT_TCLS
-	// Make sure Tenproxy.dll is loaded
-	TclsProxyConnection::Get().Init();
-#endif
-
-#if WITH_TENCENT_CROSS_SDK
-	FCrossVoiceClient::Get().Init(this);
-#endif
 
 	// Initialize Rail SDK
 	if (!InitRailSdk())
@@ -254,13 +238,9 @@ bool FOnlineSubsystemTencent::Init()
 	UE_LOG_ONLINE(Verbose, TEXT("Created thread (ID:%d)."), OnlineAsyncTaskThread->GetThreadID());
 
 	TencentIdentity = MakeShared<FOnlineIdentityTencent, ESPMode::ThreadSafe>(this);
-	TencentDirectory = MakeShared<FOnlineDirectoryTencent, ESPMode::ThreadSafe>(this);
 
-#if !WITH_TENCENT_RAIL_SDK
-	TencentSession = MakeShared<FOnlineSessionTencentTCLS, ESPMode::ThreadSafe>(this);
-#else
+#if WITH_TENCENT_RAIL_SDK
 	TencentSession = MakeShared<FOnlineSessionTencentRail, ESPMode::ThreadSafe>(this);
-#endif
 
 	if (!TencentSession->Init())
 	{
@@ -269,7 +249,6 @@ bool FOnlineSubsystemTencent::Init()
 		bInitSuccess = false;
 	}
 	
-#if WITH_TENCENT_RAIL_SDK
 	TencentFriends = MakeShared<FOnlineFriendsTencent, ESPMode::ThreadSafe>(this);
 	if (!TencentFriends->Init())
 	{
@@ -291,13 +270,6 @@ bool FOnlineSubsystemTencent::Init()
 	TencentStore = MakeShared<FOnlineStoreTencent, ESPMode::ThreadSafe>(this);
 #endif
 
-#if WITH_TENCENTSDK
-	// Get the server info from TDIR at startup
-	ExecuteNextTick([this]() {
-		TencentDirectory->QueryServerInfo();
-	});
-#endif // WITH_TENCENTSDK
-
 	// update services based on user login/logout events
 	OnLoginChangedHandle = TencentIdentity->AddOnLoginChangedDelegate_Handle(FOnLoginChangedDelegate::CreateThreadSafeSP(this, &FOnlineSubsystemTencent::OnLoginChanged));
 
@@ -439,26 +411,6 @@ void FOnlineSubsystemTencent::ShutdownRailSdk()
 
 void FOnlineSubsystemTencent::OnLoginChanged(int32 LocalUserNum)
 {
-	if (TencentIdentity.IsValid())
-	{
-		ELoginStatus::Type LoginStatus = TencentIdentity->GetLoginStatus(LocalUserNum);
-		if (LoginStatus == ELoginStatus::LoggedIn)
-		{
-			TSharedPtr<const FUniqueNetId> UserId = TencentIdentity->GetUniquePlayerId(LocalUserNum);
-			if (UserId.IsValid())
-			{
-#if WITH_TENCENT_CROSS_SDK
-				FCrossVoiceClient::Get().HandleLogin(UserId->ToString());
-#endif
-			}
-		}
-		else if (LoginStatus == ELoginStatus::NotLoggedIn)
-		{
-#if WITH_TENCENT_CROSS_SDK
-			FCrossVoiceClient::Get().HandleLogout();
-#endif
-		}
-	}
 }
 
 void FOnlineSubsystemTencent::PreUnload()
@@ -507,7 +459,6 @@ bool FOnlineSubsystemTencent::Shutdown()
 	DESTRUCT_INTERFACE(TencentMessageSanitizer);
 #endif
 	DESTRUCT_INTERFACE(TencentSession);
-	DESTRUCT_INTERFACE(TencentDirectory);
 	DESTRUCT_INTERFACE(TencentIdentity);
 
 	if (OnlineAsyncTaskThread)
@@ -523,16 +474,6 @@ bool FOnlineSubsystemTencent::Shutdown()
 		OnlineAsyncTaskThreadRunnable = nullptr;
 	}
 
-#if WITH_TENCENT_CROSS_SDK
-	// Cleanup Cross SDK
-	FCrossVoiceClient::Get().Shutdown();
-#endif
-
-#if WITH_TENCENT_TCLS
-	// Close any connection and unload Tenproxy.dll
-	TclsProxyConnection::Get().Shutdown();
-#endif
-
 #if WITH_TENCENT_RAIL_SDK
 	// Shutdown / unload the Rail SDK (Wegame)
 	ShutdownRailSdk();
@@ -568,11 +509,6 @@ FString FOnlineSubsystemTencent::GetAppId() const
 	FString AppId;
 #if WITH_TENCENT_RAIL_SDK
 	AppId = LexToString(RailGameId);
-#else
-	if (TencentDirectory.IsValid())
-	{
-		AppId = LexToString(TencentDirectory->GetAppServiceId());
-	}
 #endif
 	return AppId;
 }
@@ -593,11 +529,6 @@ bool FOnlineSubsystemTencent::Exec(UWorld* InWorld, const TCHAR* Cmd, FOutputDev
 			bWasHandled = HandleAuthExecCommands(InWorld, Cmd, Ar);
 		}
 #endif // WITH_DEV_AUTOMATION_TESTS
-
-		if (FParse::Command(&Cmd, TEXT("CROSS")))
-		{
-			bWasHandled = HandleCrossVoiceCommands(InWorld, Cmd, Ar);
-		}
 	}
 	else if (FParse::Command(&Cmd, TEXT("SESSION")))
 	{
@@ -715,116 +646,6 @@ bool FOnlineSubsystemTencent::HandleAuthExecCommands(UWorld* InWorld, const TCHA
 	return bWasHandled;
 }
 
-bool FOnlineSubsystemTencent::HandleCrossVoiceCommands(UWorld* InWorld, const TCHAR* Cmd, FOutputDevice& Ar)
-{
-	bool bWasHandled = false;
-	
-#if WITH_TENCENT_CROSS_SDK
-	FCrossVoiceClient& Cross = FCrossVoiceClient::Get();
-
-	if (FParse::Command(&Cmd, TEXT("LOGIN")))
-	{
-		IOnlineIdentityPtr IdentityInt = GetIdentityInterface();
-		if (IdentityInt.IsValid())
-		{
-			TSharedPtr<const FUniqueNetId> UserId = IdentityInt->GetUniquePlayerId(0);
-			if (UserId.IsValid())
-			{
-				Cross.HandleLogin(*UserId->ToString());
-			}
-		}
-	}
-	else if (FParse::Command(&Cmd, TEXT("LOGOUT")))
-	{
-		Cross.HandleLogout();
-	}
-	else if (FParse::Command(&Cmd, TEXT("NICK")))
-	{
-		FString NicknameStr = FParse::Token(Cmd, false);
-		if (!NicknameStr.IsEmpty())
-		{
-			IOnlineIdentityPtr IdentityInt = GetIdentityInterface();
-			if (IdentityInt.IsValid())
-			{
-				TSharedPtr<const FUniqueNetId> UserId = IdentityInt->GetUniquePlayerId(0);
-				if (UserId.IsValid())
-				{
-					Cross.SendUserInfo(*UserId, NicknameStr);
-				}
-			}
-		}
-	}
-	else if (FParse::Command(&Cmd, TEXT("ENTER")))
-	{
-		FString RoomIdStr = FParse::Token(Cmd, false);
-		uint64 RoomId;
-		LexFromString(RoomId, *RoomIdStr);
-		if (RoomId != 0 && RoomId != INVALID_VOICE_CHATROOM)
-		{
-			Cross.EnterRoom(RoomId, FOnVoiceChatRoomEntered::CreateLambda(
-				[](const FVoiceChatRoomEntered& RoomData, const FOnlineError& Error)
-			{
-				UE_LOG_ONLINE(Verbose, TEXT("EnterRoom complete %s"), *Error.ToLogString());
-			}));
-		}
-
-		bWasHandled = true;
-	}
-	else if (FParse::Command(&Cmd, TEXT("SWITCH")))
-	{
-		FString RoomIdStr = FParse::Token(Cmd, false);
-		uint64 RoomId;
-		LexFromString(RoomId, *RoomIdStr);
-		if (RoomId != 0 && RoomId != INVALID_VOICE_CHATROOM)
-		{
-			Cross.SwitchRoom(RoomId, FOnVoiceChatRoomEntered::CreateLambda(
-				[](const FVoiceChatRoomEntered& RoomData, const FOnlineError& Error)
-			{
-				UE_LOG_ONLINE(Verbose, TEXT("SwitchRoom complete %s"), *Error.ToLogString());
-			}));
-		}
-
-		bWasHandled = true;
-	}
-	else if (FParse::Command(&Cmd, TEXT("EXIT")))
-	{
-		FString RoomIdStr = FParse::Token(Cmd, false);
-		uint64 RoomId;
-		LexFromString(RoomId, *RoomIdStr);
-		if (RoomId != 0 && RoomId != INVALID_VOICE_CHATROOM)
-		{
-			Cross.ExitRoom(RoomId, FOnVoiceChatRoomExited::CreateLambda(
-				[](const FVoiceChatRoomExited& RoomData, const FOnlineError& Error)
-			{
-				UE_LOG_ONLINE(Verbose, TEXT("ExitRoom complete %s"), *Error.ToLogString());
-			}));
-		}
-
-		bWasHandled = true;
-	}
-	else if (FParse::Command(&Cmd, TEXT("SHOW")))
-	{
-		FString ShowUI = FParse::Token(Cmd, false);
-
-		bool bShowUI = false;
-		LexFromString(bShowUI, *ShowUI);
-		Cross.ShowUI(bShowUI, FOnShowVoiceChatUI::CreateLambda(
-			[](const FOnlineError& Error)
-		{
-			UE_LOG_ONLINE(Verbose, TEXT("ShowUI complete %s"), *Error.ToLogString());
-		}));
-		bWasHandled = true;
-	}
-	else if (FParse::Command(&Cmd, TEXT("DUMP")))
-	{
-		Cross.DumpState();
-		bWasHandled = true;
-	}
-#endif
-
-	return bWasHandled;
-}
-
 bool FOnlineSubsystemTencent::HandleSessionExecCommands(UWorld* InWorld, const TCHAR* Cmd, FOutputDevice& Ar)
 {
 	bool bWasHandled = false;
diff --git a/Engine/Plugins/Online/OnlineSubsystemTencent/Source/Public/OnlineIdentityTencent.h b/Engine/Plugins/Online/OnlineSubsystemTencent/Source/Public/OnlineIdentityTencent.h
--- a/Engine/Plugins/Online/OnlineSubsystemTencent/Source/Public/OnlineIdentityTencent.h
+++ b/Engine/Plugins/Online/OnlineSubsystemTencent/Source/Public/OnlineIdentityTencent.h
@@ -75,12 +75,7 @@ public:
 	/**
 	 * Init/default constructor
 	 */
-#if WITH_TENCENT_TCLS
-	FUserOnlineAccountTencent(const FString& InUserId=TEXT("")) 
-		: UserId(new FUniqueNetIdTCLS(InUserId))
-		, QQId(0)
-	{ }
-#elif WITH_TENCENT_RAIL_SDK
+#if WITH_TENCENT_RAIL_SDK
 	FUserOnlineAccountTencent(const TSharedRef<const FUniqueNetId> InUserId)
 		: UserId(InUserId)
 	{ }
@@ -94,12 +89,6 @@ public:
 	TSharedRef<const FUniqueNetId> UserId;
 	/** Any addition account data associated with the user */
 	FJsonSerializableKeyValueMap AccountData;
-#if WITH_TENCENT_TCLS
-	/** id of the QQ account. Used for UserId as well */
-	uint32 QQId;
-	/** Game signature code which can be used to verify identity in Tencent environment */
-	FString GameSignature;
-#endif
 	/** Full auth code which can be exchanged for access */
 	FString AuthToken;
 };
diff --git a/Engine/Plugins/Online/OnlineSubsystemTencent/Source/Public/OnlineSubsystemTencent.h b/Engine/Plugins/Online/OnlineSubsystemTencent/Source/Public/OnlineSubsystemTencent.h
--- a/Engine/Plugins/Online/OnlineSubsystemTencent/Source/Public/OnlineSubsystemTencent.h
+++ b/Engine/Plugins/Online/OnlineSubsystemTencent/Source/Public/OnlineSubsystemTencent.h
@@ -99,7 +99,6 @@ public:
 	// FOnlineSubsystemTencent
 
 	virtual ~FOnlineSubsystemTencent();
-	virtual FOnlineDirectoryTencentPtr GetDirectoryTencent() { return TencentDirectory; }	
 
 	/**
 	 * Delegate called when Anti Addiction dialog should be displayed
@@ -137,7 +136,6 @@ PACKAGE_SCOPE:
 private:
 
 	bool HandleAuthExecCommands(UWorld* InWorld, const TCHAR* Cmd, FOutputDevice& Ar);
-	bool HandleCrossVoiceCommands(UWorld* InWorld, const TCHAR* Cmd, FOutputDevice& Ar);
 	bool HandleSessionExecCommands(UWorld* InWorld, const TCHAR* Cmd, FOutputDevice& Ar);
 	bool HandlePresenceExecCommands(UWorld* InWorld, const TCHAR* Cmd, FOutputDevice& Ar);
 	bool HandleUsersExecCommands(UWorld* InWorld, const TCHAR* Cmd, FOutputDevice& Ar);
@@ -171,10 +169,6 @@ private:
 	FOnlineIdentityTencentPtr TencentIdentity;
 	/** Interface to the session services (needed for TSS anticheat handling) */
 	FOnlineSessionTencentPtr TencentSession;
-#if WITH_TENCENTSDK
-	/** Interface to the directory service (TCLS proxy TDIR info) */
-	FOnlineDirectoryTencentPtr TencentDirectory;
-#endif // WITH_TENCENTSDK
 #if WITH_TENCENT_RAIL_SDK
 	/** Interface to the friends services */
 	FOnlineFriendsTencentPtr TencentFriends;
diff --git a/Engine/Plugins/Online/OnlineSubsystemTencent/Source/Public/OnlineSubsystemTencentTypes.h b/Engine/Plugins/Online/OnlineSubsystemTencent/Source/Public/OnlineSubsystemTencentTypes.h
--- a/Engine/Plugins/Online/OnlineSubsystemTencent/Source/Public/OnlineSubsystemTencentTypes.h
+++ b/Engine/Plugins/Online/OnlineSubsystemTencent/Source/Public/OnlineSubsystemTencentTypes.h
@@ -180,10 +180,6 @@ public:
 
 #endif // WITH_TENCENT_RAIL_SDK
 
-#if WITH_TENCENT_TCLS
-	TEMP_UNIQUENETIDSTRING_SUBCLASS(FUniqueNetIdTCLS, TENCENT_SUBSYSTEM);
-#endif // WITH_TENCENT_TCLS
-
 class FOnlineSessionInfoTencent : public FOnlineSessionInfo
 {
 private:
