Added a "NotImplemented" implementation of the new Checkout method whose callback has no receipts for Tencent

diff --git a/Engine/Plugins/Online/OnlineSubsystemTencent/Source/Private/OnlinePurchaseTencent.cpp b/Engine/Plugins/Online/OnlineSubsystemTencent/Source/Private/OnlinePurchaseTencent.cpp
--- a/Engine/Plugins/Online/OnlineSubsystemTencent/Source/Private/OnlinePurchaseTencent.cpp
+++ b/Engine/Plugins/Online/OnlineSubsystemTencent/Source/Private/OnlinePurchaseTencent.cpp
@@ -92,6 +92,16 @@ void FOnlinePurchaseTencent::Checkout_Complete(const FPurchaseProductsTaskResult
 	Delegate.ExecuteIfBound(Result.Error, Result.Error.WasSuccessful() ? Result.PurchaseReceipt.ToSharedRef() : EmptyReceipt);
 }
 
+void FOnlinePurchaseTencent::Checkout(const FUniqueNetId& UserId, const FPurchaseCheckoutRequest& CheckoutRequest, const FOnPurchaseReceiptlessCheckoutComplete& Delegate)
+{
+	// NYI - please use the other Checkout method
+	Subsystem->ExecuteNextTick([Delegate]()
+	{
+		static const FOnlineError NotImplementedResult(EOnlineErrorResult::NotImplemented);
+		Delegate.ExecuteIfBound(NotImplementedResult);
+	});
+}
+
 void FOnlinePurchaseTencent::FinalizePurchase(const FUniqueNetId& UserId, const FString& ReceiptId)
 {
 	// NYI
diff --git a/Engine/Plugins/Online/OnlineSubsystemTencent/Source/Private/OnlinePurchaseTencent.h b/Engine/Plugins/Online/OnlineSubsystemTencent/Source/Private/OnlinePurchaseTencent.h
--- a/Engine/Plugins/Online/OnlineSubsystemTencent/Source/Private/OnlinePurchaseTencent.h
+++ b/Engine/Plugins/Online/OnlineSubsystemTencent/Source/Private/OnlinePurchaseTencent.h
@@ -24,6 +24,7 @@ public:
 	// Begin IOnlinePurchase
 	virtual bool IsAllowedToPurchase(const FUniqueNetId& UserId) override;
 	virtual void Checkout(const FUniqueNetId& UserId, const FPurchaseCheckoutRequest& CheckoutRequest, const FOnPurchaseCheckoutComplete& Delegate) override;
+	virtual void Checkout(const FUniqueNetId& UserId, const FPurchaseCheckoutRequest& CheckoutRequest, const FOnPurchaseReceiptlessCheckoutComplete& Delegate) override;
 	virtual void FinalizePurchase(const FUniqueNetId& UserId, const FString& ReceiptId) override;
 	virtual void RedeemCode(const FUniqueNetId& UserId, const FRedeemCodeRequest& RedeemCodeRequest, const FOnPurchaseRedeemCodeComplete& Delegate) override;
 	virtual void QueryReceipts(const FUniqueNetId& UserId, bool bRestoreReceipts, const FOnQueryReceiptsComplete& Delegate) override;
