Renamed private pch header in OnlineSubsytemTencent and fixed -NoPCH -DisableUnity compile issues.

diff --git a/Engine/Plugins/Online/OnlineSubsystemTencent/Source/OnlineSubsystemTencent.Build.cs b/Engine/Plugins/Online/OnlineSubsystemTencent/Source/OnlineSubsystemTencent.Build.cs
--- a/Engine/Plugins/Online/OnlineSubsystemTencent/Source/OnlineSubsystemTencent.Build.cs
+++ b/Engine/Plugins/Online/OnlineSubsystemTencent/Source/OnlineSubsystemTencent.Build.cs
@@ -9,8 +9,6 @@ public class OnlineSubsystemTencent : ModuleRules
 	{
 		PrivateDefinitions.Add("ONLINESUBSYSTEM_TENCENT_PACKAGE=1");
 
-		PCHUsage = ModuleRules.PCHUsageMode.UseExplicitOrSharedPCHs;
-
 		PrivateIncludePaths.Add("Private");
 
 		PublicDependencyModuleNames.AddRange(
diff --git a/Engine/Plugins/Online/OnlineSubsystemTencent/Source/Private/OnlineAsyncTaskManagerTencent.cpp b/Engine/Plugins/Online/OnlineSubsystemTencent/Source/Private/OnlineAsyncTaskManagerTencent.cpp
--- a/Engine/Plugins/Online/OnlineSubsystemTencent/Source/Private/OnlineAsyncTaskManagerTencent.cpp
+++ b/Engine/Plugins/Online/OnlineSubsystemTencent/Source/Private/OnlineAsyncTaskManagerTencent.cpp
@@ -1,7 +1,7 @@
 // Copyright Epic Games, Inc. All Rights Reserved.
 
 #include "OnlineAsyncTaskManagerTencent.h"
-#include "OnlineSubsystemTencentPCH.h"
+#include "OnlineSubsystemTencentPrivate.h"
 #include "OnlineAsyncTasksTencent.h"
 #include "OnlinePresenceTencent.h"
 #include "PlayTimeLimitImpl.h"
diff --git a/Engine/Plugins/Online/OnlineSubsystemTencent/Source/Private/OnlineAsyncTasksTencent.cpp b/Engine/Plugins/Online/OnlineSubsystemTencent/Source/Private/OnlineAsyncTasksTencent.cpp
--- a/Engine/Plugins/Online/OnlineSubsystemTencent/Source/Private/OnlineAsyncTasksTencent.cpp
+++ b/Engine/Plugins/Online/OnlineSubsystemTencent/Source/Private/OnlineAsyncTasksTencent.cpp
@@ -1,7 +1,7 @@
 // Copyright Epic Games, Inc. All Rights Reserved.
 
 #include "OnlineAsyncTasksTencent.h"
-#include "OnlineSubsystemTencentPCH.h"
+#include "OnlineSubsystemTencentPrivate.h"
 #include "OnlineSubsystemTencentTypes.h"
 #include "Interfaces/OnlineExternalUIInterface.h"
 #include "Interfaces/OnlinePresenceInterface.h"
diff --git a/Engine/Plugins/Online/OnlineSubsystemTencent/Source/Private/OnlineFriendsTencent.cpp b/Engine/Plugins/Online/OnlineSubsystemTencent/Source/Private/OnlineFriendsTencent.cpp
--- a/Engine/Plugins/Online/OnlineSubsystemTencent/Source/Private/OnlineFriendsTencent.cpp
+++ b/Engine/Plugins/Online/OnlineSubsystemTencent/Source/Private/OnlineFriendsTencent.cpp
@@ -1,7 +1,7 @@
 // Copyright Epic Games, Inc. All Rights Reserved.
 
 #include "OnlineFriendsTencent.h"
-#include "OnlineSubsystemTencentPCH.h"
+#include "OnlineSubsystemTencentPrivate.h"
 #include "OnlineIdentityTencent.h"
 #include "OnlinePresenceTencent.h"
 #include "OnlineAsyncTasksTencent.h"
diff --git a/Engine/Plugins/Online/OnlineSubsystemTencent/Source/Private/OnlineIdentityTencentCommon.cpp b/Engine/Plugins/Online/OnlineSubsystemTencent/Source/Private/OnlineIdentityTencentCommon.cpp
--- a/Engine/Plugins/Online/OnlineSubsystemTencent/Source/Private/OnlineIdentityTencentCommon.cpp
+++ b/Engine/Plugins/Online/OnlineSubsystemTencent/Source/Private/OnlineIdentityTencentCommon.cpp
@@ -1,7 +1,7 @@
 // Copyright Epic Games, Inc. All Rights Reserved.
 
 #include "OnlineIdentityTencent.h"
-#include "OnlineSubsystemTencentPCH.h"
+#include "OnlineSubsystemTencentPrivate.h"
 
 #if WITH_TENCENT_RAIL_SDK
 bool FOnlineIdentityTencent::GetLocalUserIdx(const FUniqueNetId& UserId, int32& OutLocalIdx) const
diff --git a/Engine/Plugins/Online/OnlineSubsystemTencent/Source/Private/OnlineIdentityTencentRail.cpp b/Engine/Plugins/Online/OnlineSubsystemTencent/Source/Private/OnlineIdentityTencentRail.cpp
--- a/Engine/Plugins/Online/OnlineSubsystemTencent/Source/Private/OnlineIdentityTencentRail.cpp
+++ b/Engine/Plugins/Online/OnlineSubsystemTencent/Source/Private/OnlineIdentityTencentRail.cpp
@@ -1,7 +1,7 @@
 // Copyright Epic Games, Inc. All Rights Reserved.
 
 #include "OnlineIdentityTencent.h"
-#include "OnlineSubsystemTencentPCH.h"
+#include "OnlineSubsystemTencentPrivate.h"
 
 #if WITH_TENCENTSDK
 #if WITH_TENCENT_RAIL_SDK
diff --git a/Engine/Plugins/Online/OnlineSubsystemTencent/Source/Private/OnlinePlayTimeLimitTencent.cpp b/Engine/Plugins/Online/OnlineSubsystemTencent/Source/Private/OnlinePlayTimeLimitTencent.cpp
--- a/Engine/Plugins/Online/OnlineSubsystemTencent/Source/Private/OnlinePlayTimeLimitTencent.cpp
+++ b/Engine/Plugins/Online/OnlineSubsystemTencent/Source/Private/OnlinePlayTimeLimitTencent.cpp
@@ -1,7 +1,7 @@
 // Copyright Epic Games, Inc. All Rights Reserved.
 
 #include "OnlinePlayTimeLimitTencent.h"
-#include "OnlineSubsystemTencentPCH.h"
+#include "OnlineSubsystemTencentPrivate.h"
 
 #if WITH_TENCENT_RAIL_SDK
 
diff --git a/Engine/Plugins/Online/OnlineSubsystemTencent/Source/Private/OnlinePresenceTencent.cpp b/Engine/Plugins/Online/OnlineSubsystemTencent/Source/Private/OnlinePresenceTencent.cpp
--- a/Engine/Plugins/Online/OnlineSubsystemTencent/Source/Private/OnlinePresenceTencent.cpp
+++ b/Engine/Plugins/Online/OnlineSubsystemTencent/Source/Private/OnlinePresenceTencent.cpp
@@ -1,7 +1,7 @@
 // Copyright Epic Games, Inc. All Rights Reserved.
 
 #include "OnlinePresenceTencent.h"
-#include "OnlineSubsystemTencentPCH.h"
+#include "OnlineSubsystemTencentPrivate.h"
 
 #if WITH_TENCENTSDK
 #if WITH_TENCENT_RAIL_SDK
diff --git a/Engine/Plugins/Online/OnlineSubsystemTencent/Source/Private/OnlineSessionTencent.cpp b/Engine/Plugins/Online/OnlineSubsystemTencent/Source/Private/OnlineSessionTencent.cpp
--- a/Engine/Plugins/Online/OnlineSubsystemTencent/Source/Private/OnlineSessionTencent.cpp
+++ b/Engine/Plugins/Online/OnlineSubsystemTencent/Source/Private/OnlineSessionTencent.cpp
@@ -1,7 +1,7 @@
 // Copyright Epic Games, Inc. All Rights Reserved.
 
 #include "OnlineSessionTencent.h"
-#include "OnlineSubsystemTencentPCH.h"
+#include "OnlineSubsystemTencentPrivate.h"
 #include "OnlineSubsystemTencentTypes.h"
 #include "OnlineIdentityTencent.h"
 #include "Misc/ScopeLock.h"
diff --git a/Engine/Plugins/Online/OnlineSubsystemTencent/Source/Private/OnlineSessionTencentRail.cpp b/Engine/Plugins/Online/OnlineSubsystemTencent/Source/Private/OnlineSessionTencentRail.cpp
--- a/Engine/Plugins/Online/OnlineSubsystemTencent/Source/Private/OnlineSessionTencentRail.cpp
+++ b/Engine/Plugins/Online/OnlineSubsystemTencent/Source/Private/OnlineSessionTencentRail.cpp
@@ -1,7 +1,7 @@
 // Copyright Epic Games, Inc. All Rights Reserved.
 
 #include "OnlineSessionTencentRail.h"
-#include "OnlineSubsystemTencentPCH.h"
+#include "OnlineSubsystemTencentPrivate.h"
 #include "OnlinePresenceTencent.h"
 #include "Misc/CommandLine.h"
 #include "Misc/App.h"
diff --git a/Engine/Plugins/Online/OnlineSubsystemTencent/Source/Private/OnlineSubsystemModuleTencent.cpp b/Engine/Plugins/Online/OnlineSubsystemTencent/Source/Private/OnlineSubsystemModuleTencent.cpp
--- a/Engine/Plugins/Online/OnlineSubsystemTencent/Source/Private/OnlineSubsystemModuleTencent.cpp
+++ b/Engine/Plugins/Online/OnlineSubsystemTencent/Source/Private/OnlineSubsystemModuleTencent.cpp
@@ -1,7 +1,7 @@
 // Copyright Epic Games, Inc. All Rights Reserved.
 
 #include "OnlineSubsystemTencentModule.h"
-#include "OnlineSubsystemTencentPCH.h"
+#include "OnlineSubsystemTencentPrivate.h"
 #include "OnlineSubsystemTencent.h"
 
 IMPLEMENT_MODULE(FOnlineSubsystemTencentModule, OnlineSubsystemTencent);
diff --git a/Engine/Plugins/Online/OnlineSubsystemTencent/Source/Private/OnlineSubsystemTencent.cpp b/Engine/Plugins/Online/OnlineSubsystemTencent/Source/Private/OnlineSubsystemTencent.cpp
--- a/Engine/Plugins/Online/OnlineSubsystemTencent/Source/Private/OnlineSubsystemTencent.cpp
+++ b/Engine/Plugins/Online/OnlineSubsystemTencent/Source/Private/OnlineSubsystemTencent.cpp
@@ -1,7 +1,7 @@
 // Copyright Epic Games, Inc. All Rights Reserved.
 
 #include "OnlineSubsystemTencent.h"
-#include "OnlineSubsystemTencentPCH.h"
+#include "OnlineSubsystemTencentPrivate.h"
 #include "RailSdkWrapper.h"
 #include "OnlineAsyncTaskManagerTencent.h"
 #include "OnlineIdentityTencent.h"
diff --git a/Engine/Plugins/Online/OnlineSubsystemTencent/Source/Private/OnlineSubsystemTencentPCH.h b/Engine/Plugins/Online/OnlineSubsystemTencent/Source/Private/OnlineSubsystemTencentPCH.h
deleted file mode 100644
diff --git a/Engine/Plugins/Online/OnlineSubsystemTencent/Source/Private/OnlineSubsystemTencentPrivate.h b/Engine/Plugins/Online/OnlineSubsystemTencent/Source/Private/OnlineSubsystemTencentPrivate.h
new file mode 100644
--- /dev/null
+++ b/Engine/Plugins/Online/OnlineSubsystemTencent/Source/Private/OnlineSubsystemTencentPrivate.h
@@ -0,0 +1,18 @@
+// Copyright Epic Games, Inc. All Rights Reserved.
+
+#pragma once
+
+#include "OnlineSubsystem.h"
+
+#include "OnlineSubsystemTencent.h"
+#include "OnlineSubsystemTencentTypes.h"
+
+#if WITH_ENGINE
+#include "OnlineSubsystemUtils.h"
+#endif
+
+#define RAIL_METADATA_KEY_SEPARATOR TEXT(",")
+
+/** pre-pended to all mcp logging */
+#undef ONLINE_LOG_PREFIX
+#define ONLINE_LOG_PREFIX TEXT("TENCENT: ")
diff --git a/Engine/Plugins/Online/OnlineSubsystemTencent/Source/Private/OnlineUserTencent.cpp b/Engine/Plugins/Online/OnlineSubsystemTencent/Source/Private/OnlineUserTencent.cpp
--- a/Engine/Plugins/Online/OnlineSubsystemTencent/Source/Private/OnlineUserTencent.cpp
+++ b/Engine/Plugins/Online/OnlineSubsystemTencent/Source/Private/OnlineUserTencent.cpp
@@ -1,7 +1,7 @@
 // Copyright Epic Games, Inc. All Rights Reserved.
 
 #include "OnlineUserTencent.h"
-#include "OnlineSubsystemTencentPCH.h"
+#include "OnlineSubsystemTencentPrivate.h"
 #include "OnlineAsyncTasksTencent.h"
 
 #if WITH_TENCENTSDK
diff --git a/Engine/Plugins/Online/OnlineSubsystemTencent/Source/Private/RailSDK.cpp b/Engine/Plugins/Online/OnlineSubsystemTencent/Source/Private/RailSDK.cpp
--- a/Engine/Plugins/Online/OnlineSubsystemTencent/Source/Private/RailSDK.cpp
+++ b/Engine/Plugins/Online/OnlineSubsystemTencent/Source/Private/RailSDK.cpp
@@ -1,7 +1,7 @@
 // Copyright Epic Games, Inc. All Rights Reserved.
 
 #include "RailSDK.h"
-#include "OnlineSubsystemTencentPCH.h"
+#include "OnlineSubsystemTencentPrivate.h"
 #include "RailSdkWrapper.h"
 
 #if WITH_TENCENTSDK
diff --git a/Engine/Plugins/Online/OnlineSubsystemTencent/Source/Private/RailSdkWrapper.cpp b/Engine/Plugins/Online/OnlineSubsystemTencent/Source/Private/RailSdkWrapper.cpp
--- a/Engine/Plugins/Online/OnlineSubsystemTencent/Source/Private/RailSdkWrapper.cpp
+++ b/Engine/Plugins/Online/OnlineSubsystemTencent/Source/Private/RailSdkWrapper.cpp
@@ -1,7 +1,7 @@
 // Copyright Epic Games, Inc. All Rights Reserved.
 
 #include "RailSdkWrapper.h"
-#include "OnlineSubsystemTencentPCH.h"
+#include "OnlineSubsystemTencentPrivate.h"
 #include "Misc/Paths.h"
 #include "Misc/CommandLine.h"
 #include "HAL/PlatformProcess.h"
diff --git a/Engine/Plugins/Online/OnlineSubsystemTencent/Source/Private/TencentDllMgr.cpp b/Engine/Plugins/Online/OnlineSubsystemTencent/Source/Private/TencentDllMgr.cpp
--- a/Engine/Plugins/Online/OnlineSubsystemTencent/Source/Private/TencentDllMgr.cpp
+++ b/Engine/Plugins/Online/OnlineSubsystemTencent/Source/Private/TencentDllMgr.cpp
@@ -1,7 +1,7 @@
 // Copyright Epic Games, Inc. All Rights Reserved.
 
 #include "TencentDllMgr.h"
-#include "OnlineSubsystemTencentPCH.h"
+#include "OnlineSubsystemTencentPrivate.h"
 #include "Misc/Paths.h"
 
 FTencentDll::FTencentDll()
