- Fix build fail in FUniqueNetIdRail

diff --git a/Engine/Plugins/Online/OnlineSubsystemTencent/Source/Public/OnlineSubsystemTencentTypes.h b/Engine/Plugins/Online/OnlineSubsystemTencent/Source/Public/OnlineSubsystemTencentTypes.h
--- a/Engine/Plugins/Online/OnlineSubsystemTencent/Source/Public/OnlineSubsystemTencentTypes.h
+++ b/Engine/Plugins/Online/OnlineSubsystemTencent/Source/Public/OnlineSubsystemTencentTypes.h
@@ -124,7 +124,7 @@ public:
 
 	virtual uint32 GetTypeHash() const override
 	{
-		uint64 id = A.RailID.get_id();
+		uint64 id = RailID.get_id();
 		return ::GetTypeHash(id);
 	}
 
