Fixed some 'deprecated' FString usage.

diff --git a/Engine/Plugins/Online/OnlineSubsystemTencent/Source/Private/RailSDK.cpp b/Engine/Plugins/Online/OnlineSubsystemTencent/Source/Private/RailSDK.cpp
--- a/Engine/Plugins/Online/OnlineSubsystemTencent/Source/Private/RailSDK.cpp
+++ b/Engine/Plugins/Online/OnlineSubsystemTencent/Source/Private/RailSDK.cpp
@@ -119,7 +119,7 @@ FString LexToString(const rail::EnumRailAssetState AssetState)
 FString LexToString(const rail::RailString& RailString)
 {
 	FUTF8ToTCHAR Converter(RailString.c_str(), RailString.size());
-	return FString(Converter.Length(), Converter.Get());
+	return FString::ConstructFromPtrSize(Converter.Get(), Converter.Length());
 }
 
 void ToRailString(const FString& Str, rail::RailString& OutString)
